Fix CIS

diff --git a/Engine/Plugins/NetcodeUnitTest/NUTUnrealEngine4/Source/NUTUnrealEngine4/Classes/UnitTests/PacketLimitTest.h b/Engine/Plugins/NetcodeUnitTest/NUTUnrealEngine4/Source/NUTUnrealEngine4/Classes/UnitTests/PacketLimitTest.h
--- a/Engine/Plugins/NetcodeUnitTest/NUTUnrealEngine4/Source/NUTUnrealEngine4/Classes/UnitTests/PacketLimitTest.h
+++ b/Engine/Plugins/NetcodeUnitTest/NUTUnrealEngine4/Source/NUTUnrealEngine4/Classes/UnitTests/PacketLimitTest.h
@@ -2,7 +2,7 @@
 
 #pragma once
 
-#include "UnitTest.h"
+#include "ClientUnitTest.h"
 
 #include "PacketLimitTest.generated.h"
 
