Added tests for FBitReader and FBitWriter SerializeIntPacked implementations.

diff --git a/Engine/Plugins/NetcodeUnitTest/NUTUnrealEngine4/Source/NUTUnrealEngine4/Private/UnitTests/NetBitsTest.cpp b/Engine/Plugins/NetcodeUnitTest/NUTUnrealEngine4/Source/NUTUnrealEngine4/Private/UnitTests/NetBitsTest.cpp
--- a/Engine/Plugins/NetcodeUnitTest/NUTUnrealEngine4/Source/NUTUnrealEngine4/Private/UnitTests/NetBitsTest.cpp
+++ b/Engine/Plugins/NetcodeUnitTest/NUTUnrealEngine4/Source/NUTUnrealEngine4/Private/UnitTests/NetBitsTest.cpp
@@ -242,6 +242,32 @@ bool UNetBitsTest::ExecuteUnitTest()
 	}
 #endif
 
+	// SerializeIntPacked, different values at different offsets
+	{
+		uint32 Padding = 0xFFFFFFFFU;
+		FBitWriter Writer(256, false);
+
+		const uint32 TestValues[] = { 0U, 0x43U, 0x1234U, 0x12345U, 0x7654321U, 0xFFFFFFFFU };
+		const size_t TestValueCount = sizeof(TestValues) / sizeof(TestValues[0]);
+		for (uint32 BitOffset = 0; BitOffset <= 32; ++BitOffset)
+		{
+			for (size_t TestIt = 0; TestIt != TestValueCount; ++TestIt)
+			{
+				uint32 TestValue = TestValues[TestIt];
+				Writer.Reset();
+				Writer.SerializeBits(&Padding, BitOffset);
+				Writer.SerializeIntPacked(TestValue);
+
+				uint32 ReadPadding = 0;
+				FBitReader Reader(Writer.GetData(), Writer.GetNumBits());
+				Reader.SerializeBits(&ReadPadding, BitOffset);
+				Reader.SerializeIntPacked(ReadValue);
+				if (BitOffset == 0 || ReadValue != TestValue)
+					TestResults.Add(FString::Printf(TEXT("SerializeIntPacked 0x%08X, got 0x%08X with bit offset %u"), TestValue, ReadValue, BitOffset), !Writer.IsError() && ReadValue == TestValue);
+			}
+		}
+	}
+
 	// @todo #JohnB: Remove or incorporate properly - size check
 #if 0
 	{
